Added a Content Browser file data source for Portable Object (PO) files contained within the Localization folder

diff --git a/Engine/Config/BaseEditor.ini b/Engine/Config/BaseEditor.ini
--- a/Engine/Config/BaseEditor.ini
+++ b/Engine/Config/BaseEditor.ini
@@ -68,6 +68,7 @@ DefaultSizeY=512
 [/Script/ContentBrowserData.ContentBrowserDataSubsystem]
 +EnabledDataSources="AssetData"
 +EnabledDataSources="ClassData"
++EnabledDataSources="PortableObjectData"
 
 [/Script/UnrealEd.CollisionAutomationTestConfigData]
 
diff --git a/Engine/Plugins/Editor/Localization/PortableObjectFileDataSource/PortableObjectFileDataSource.uplugin b/Engine/Plugins/Editor/Localization/PortableObjectFileDataSource/PortableObjectFileDataSource.uplugin
new file mode 100644
--- /dev/null
+++ b/Engine/Plugins/Editor/Localization/PortableObjectFileDataSource/PortableObjectFileDataSource.uplugin
@@ -0,0 +1,34 @@
+{
+	"FileVersion" : 3,
+	"Version" : 1,
+	"VersionName" : "1.0",
+	"FriendlyName" : "Portable Object File Data Source",
+	"Description" : "Data Source plugin providing portable object (PO) file support for the Content Browser",
+	"Category" : "Localization",
+	"CreatedBy" : "Epic Games, Inc.",
+	"CreatedByURL" : "https://epicgames.com",
+	"DocsURL" : "",
+	"MarketplaceURL" : "",
+	"SupportURL" : "",
+	"EnabledByDefault" : true,
+	"CanContainContent" : false,
+	"IsBetaVersion" : false,
+	"Installed" : false,
+	"Modules" :
+	[
+		{
+			"Name" : "PortableObjectFileDataSource",
+			"Type" : "Editor",
+			"LoadingPhase" : "Default"
+		}
+	],
+	"Plugins": [
+		{
+			"Name": "ContentBrowserFileDataSource",
+			"Enabled": true,
+			"TargetAllowList": [
+				"Editor"
+			]
+		}
+	]
+}
\ No newline at end of file
diff --git a/Engine/Plugins/Editor/Localization/PortableObjectFileDataSource/Source/PortableObjectFileDataSource/PortableObjectFileDataSource.Build.cs b/Engine/Plugins/Editor/Localization/PortableObjectFileDataSource/Source/PortableObjectFileDataSource/PortableObjectFileDataSource.Build.cs
new file mode 100644
--- /dev/null
+++ b/Engine/Plugins/Editor/Localization/PortableObjectFileDataSource/Source/PortableObjectFileDataSource/PortableObjectFileDataSource.Build.cs
@@ -0,0 +1,26 @@
+// Copyright Epic Games, Inc. All Rights Reserved.
+
+namespace UnrealBuildTool.Rules
+{
+	public class PortableObjectFileDataSource : ModuleRules
+	{
+		public PortableObjectFileDataSource(ReadOnlyTargetRules Target) : base(Target)
+		{
+			ShortName = "POFileDataSource";
+
+			PublicDependencyModuleNames.AddRange(
+				new string[] {
+					"Core",
+					"CoreUObject",
+				}
+			);
+
+			PrivateDependencyModuleNames.AddRange(
+				new string[]
+				{
+					"ContentBrowserFileDataSource",
+				}
+			);
+		}
+	}
+}
diff --git a/Engine/Plugins/Editor/Localization/PortableObjectFileDataSource/Source/PortableObjectFileDataSource/Private/PortableObjectFileDataSourceModule.cpp b/Engine/Plugins/Editor/Localization/PortableObjectFileDataSource/Source/PortableObjectFileDataSource/Private/PortableObjectFileDataSourceModule.cpp
new file mode 100644
--- /dev/null
+++ b/Engine/Plugins/Editor/Localization/PortableObjectFileDataSource/Source/PortableObjectFileDataSource/Private/PortableObjectFileDataSourceModule.cpp
@@ -0,0 +1,129 @@
+// Copyright Epic Games, Inc. All Rights Reserved.
+
+#include "ContentBrowserFileDataCore.h"
+#include "ContentBrowserFileDataSource.h"
+#include "Misc/PackageName.h"
+#include "Misc/Paths.h"
+#include "Modules/ModuleManager.h"
+#include "UObject/Package.h"
+#include "UObject/StrongObjectPtr.h"
+
+#define LOCTEXT_NAMESPACE "PortableObjectFileData"
+
+class FPortableObjectFileDataSourceModule : public FDefaultModuleImpl
+{
+public:
+	virtual void StartupModule() override
+	{
+		if (GIsEditor && !IsRunningCommandlet())
+		{
+			ContentBrowserFileData::FFileConfigData PoFileConfig;
+			{
+				static const FText PoReadOnlyError = LOCTEXT("", "Portable Object files are managed by the localization pipeline");
+
+				auto PoCanCreate = [](const FName /*InDestFolderPath*/, const FString& /*InDestFolder*/, FText* OutErrorMsg)
+				{
+					ContentBrowserFileData::SetOptionalErrorMessage(OutErrorMsg, PoReadOnlyError);
+					return false;
+				};
+
+				auto PoCanDeleteOrDuplicate = [](const FName /*InFilePath*/, const FString& /*InFilename*/, FText* OutErrorMsg)
+				{
+					ContentBrowserFileData::SetOptionalErrorMessage(OutErrorMsg, PoReadOnlyError);
+					return false;
+				};
+
+				auto PoCanRename = [](const FName /*InFilePath*/, const FString& /*InFilename*/, const FString* /*InNewName*/, FText* OutErrorMsg)
+				{
+					ContentBrowserFileData::SetOptionalErrorMessage(OutErrorMsg, PoReadOnlyError);
+					return false;
+				};
+
+				auto PoCanCopyOrMove = [](const FName /*InFilePath*/, const FString& /*InFilename*/, const FString& /*InDestFolder*/, FText* OutErrorMsg)
+				{
+					ContentBrowserFileData::SetOptionalErrorMessage(OutErrorMsg, PoReadOnlyError);
+					return false;
+				};
+
+				ContentBrowserFileData::FDirectoryActions PoDirectoryActions;
+				PoDirectoryActions.CanCreate.BindStatic(PoCanCreate);
+				PoDirectoryActions.CanDelete.BindStatic(PoCanDeleteOrDuplicate);
+				PoDirectoryActions.CanRename.BindStatic(PoCanRename);
+				PoDirectoryActions.CanCopy.BindStatic(PoCanCopyOrMove);
+				PoDirectoryActions.CanMove.BindStatic(PoCanCopyOrMove);
+				PoDirectoryActions.PassesFilter.BindStatic(&ContentBrowserFileData::FDefaultFileActions::ItemPassesFilter, false);
+				PoDirectoryActions.GetAttribute.BindStatic(&ContentBrowserFileData::FDefaultFileActions::GetItemAttribute);
+				PoFileConfig.SetDirectoryActions(PoDirectoryActions);
+
+				ContentBrowserFileData::FFileActions PoFileActions;
+				PoFileActions.TypeExtension = TEXT("po");
+				PoFileActions.TypeName = FTopLevelAssetPath(TEXT("/Script/PortableObject.PortableObject")); // Fake path to satisfy FFileActions requirements
+				PoFileActions.TypeDisplayName = LOCTEXT("TypeName", "Portable Object");
+				PoFileActions.TypeShortDescription = LOCTEXT("TypeShortDescription", "Portable Object");
+				PoFileActions.TypeFullDescription = LOCTEXT("TypeFullDescription", "Portable Object (PO) Translation Data");
+				PoFileActions.TypeColor = FColor(200, 191, 231);
+				PoFileActions.CanCreate.BindStatic(PoCanCreate);
+				PoFileActions.CanDelete.BindStatic(PoCanDeleteOrDuplicate);
+				PoFileActions.CanRename.BindStatic(PoCanRename);
+				PoFileActions.CanCopy.BindStatic(PoCanCopyOrMove);
+				PoFileActions.CanMove.BindStatic(PoCanCopyOrMove);
+				PoFileActions.CanDuplicate.BindStatic(PoCanDeleteOrDuplicate);
+				PoFileActions.PassesFilter.BindStatic(&ContentBrowserFileData::FDefaultFileActions::ItemPassesFilter, true);
+				PoFileActions.GetAttribute.BindStatic(&ContentBrowserFileData::FDefaultFileActions::GetItemAttribute);
+				PoFileConfig.RegisterFileActions(PoFileActions);
+			}
+
+			PoFileDataSource.Reset(NewObject<UContentBrowserFileDataSource>(GetTransientPackage(), "PortableObjectData"));
+			PoFileDataSource->Initialize(PoFileConfig);
+
+			// Register the current paths that may contain localization data
+			{
+				TArray<FString> RootPaths;
+				FPackageName::QueryRootContentPaths(RootPaths);
+				for (const FString& RootPath : RootPaths)
+				{
+					OnContentPathMounted(RootPath, FPackageName::LongPackageNameToFilename(RootPath));
+				}
+			}
+			
+			// Listen for new paths that may contain localization data
+			FPackageName::OnContentPathMounted().AddRaw(this, &FPortableObjectFileDataSourceModule::OnContentPathMounted);
+			FPackageName::OnContentPathDismounted().AddRaw(this, &FPortableObjectFileDataSourceModule::OnContentPathDismounted);
+		}
+	}
+
+	virtual void ShutdownModule() override
+	{
+		FPackageName::OnContentPathMounted().RemoveAll(this);
+		FPackageName::OnContentPathDismounted().RemoveAll(this);
+			
+		PoFileDataSource.Reset();
+	}
+
+private:
+	void OnContentPathMounted(const FString& InAssetPath, const FString& InFilesystemPath)
+	{
+		if (PoFileDataSource)
+		{
+			const FString LocalizationTargetFolder = InFilesystemPath / TEXT("Localization");
+			if (FPaths::DirectoryExists(LocalizationTargetFolder))
+			{
+				PoFileDataSource->AddFileMount(*(InAssetPath / TEXT("Localization")), InFilesystemPath / TEXT("Localization"));
+			}
+		}
+	}
+
+	void OnContentPathDismounted(const FString& InAssetPath, const FString& InFilesystemPath)
+	{
+		if (PoFileDataSource)
+		{
+			PoFileDataSource->RemoveFileMount(*(InAssetPath / TEXT("Localization")));
+		}
+	}
+
+	TStrongObjectPtr<UContentBrowserFileDataSource> PoFileDataSource;
+};
+
+IMPLEMENT_MODULE(FPortableObjectFileDataSourceModule, PortableObjectFileDataSource);
+
+#undef LOCTEXT_NAMESPACE
