Fix to LOCTEXT macro without an identifier

#rnx

diff --git a/Engine/Plugins/Editor/Localization/PortableObjectFileDataSource/Source/PortableObjectFileDataSource/Private/PortableObjectFileDataSourceModule.cpp b/Engine/Plugins/Editor/Localization/PortableObjectFileDataSource/Source/PortableObjectFileDataSource/Private/PortableObjectFileDataSourceModule.cpp
--- a/Engine/Plugins/Editor/Localization/PortableObjectFileDataSource/Source/PortableObjectFileDataSource/Private/PortableObjectFileDataSourceModule.cpp
+++ b/Engine/Plugins/Editor/Localization/PortableObjectFileDataSource/Source/PortableObjectFileDataSource/Private/PortableObjectFileDataSourceModule.cpp
@@ -19,7 +19,7 @@ public:
 		{
 			ContentBrowserFileData::FFileConfigData PoFileConfig;
 			{
-				static const FText PoReadOnlyError = LOCTEXT("", "Portable Object files are managed by the localization pipeline");
+				static const FText PoReadOnlyError = LOCTEXT("PoReadOnlyError", "Portable Object files are managed by the localization pipeline");
 
 				auto PoCanCreate = [](const FName /*InDestFolderPath*/, const FString& /*InDestFolder*/, FText* OutErrorMsg)
 				{
