Allow code to override the "CanEdit" permission for Portable Object files

#rnx

diff --git a/Engine/Plugins/Editor/Localization/PortableObjectFileDataSource/Source/PortableObjectFileDataSource/Private/PortableObjectFileDataSourceModule.cpp b/Engine/Plugins/Editor/Localization/PortableObjectFileDataSource/Source/PortableObjectFileDataSource/Private/PortableObjectFileDataSourceModule.cpp
--- a/Engine/Plugins/Editor/Localization/PortableObjectFileDataSource/Source/PortableObjectFileDataSource/Private/PortableObjectFileDataSourceModule.cpp
+++ b/Engine/Plugins/Editor/Localization/PortableObjectFileDataSource/Source/PortableObjectFileDataSource/Private/PortableObjectFileDataSourceModule.cpp
@@ -1,5 +1,7 @@
 // Copyright Epic Games, Inc. All Rights Reserved.
 
+#include "IPortableObjectFileDataSourceModule.h"
+
 #include "ContentBrowserFileDataCore.h"
 #include "ContentBrowserFileDataSource.h"
 #include "LocalizationDelegates.h"
@@ -11,7 +13,7 @@
 
 #define LOCTEXT_NAMESPACE "PortableObjectFileData"
 
-class FPortableObjectFileDataSourceModule : public FDefaultModuleImpl
+class FPortableObjectFileDataSourceModule : public IPortableObjectFileDataSourceModule
 {
 public:
 	virtual void StartupModule() override
@@ -69,6 +71,7 @@ public:
 				PoFileActions.CanCopy.BindStatic(PoCanCopyOrMove);
 				PoFileActions.CanMove.BindStatic(PoCanCopyOrMove);
 				PoFileActions.CanDuplicate.BindStatic(PoCanDeleteOrDuplicate);
+				PoFileActions.CanEdit.BindRaw(this, &FPortableObjectFileDataSourceModule::CanEditFile);
 				PoFileActions.PassesFilter.BindStatic(&ContentBrowserFileData::FDefaultFileActions::ItemPassesFilter, true);
 				PoFileActions.GetAttribute.BindStatic(&ContentBrowserFileData::FDefaultFileActions::GetItemAttribute);
 				PoFileConfig.RegisterFileActions(PoFileActions);
@@ -103,6 +106,26 @@ public:
 		PoFileDataSource.Reset();
 	}
 
+	virtual FDelegateHandle RegisterCanEditFileOverride(FCanEditFileDelegate&& Delegate) override
+	{
+		FDelegateHandle Handle(FDelegateHandle::GenerateNewHandle);
+		CanEditFileOverrides.Add(MakeTuple(Handle, MoveTemp(Delegate)));
+		return Handle;
+	}
+
+	virtual void UnregisterCanEditFileOverride(const FDelegateHandle& Handle) override
+	{
+		const int32 DelegateIndex = CanEditFileOverrides.IndexOfByPredicate([&Handle](const TTuple<FDelegateHandle, FCanEditFileDelegate>& DelegatePair)
+		{
+			return DelegatePair.Key == Handle;
+		});
+
+		if (DelegateIndex != INDEX_NONE)
+		{
+			CanEditFileOverrides.RemoveAt(DelegateIndex, EAllowShrinking::No);
+		}
+	}
+
 private:
 	void OnContentPathMounted(const FString& InAssetPath, const FString& InFilesystemPath)
 	{
@@ -144,7 +167,28 @@ private:
 		}
 	}
 
+	bool CanEditFile(const FName InFilePath, const FString& InFilename, FText* OutErrorMsg)
+	{
+		for (auto DelegateIt = CanEditFileOverrides.CreateIterator(); DelegateIt; ++DelegateIt)
+		{
+			if (!DelegateIt->Value.IsBound())
+			{
+				DelegateIt.RemoveCurrent();
+				continue;
+			}
+
+			if (!DelegateIt->Value.Execute(InFilePath, InFilename, OutErrorMsg))
+			{
+				return false;
+			}
+		}
+
+		return true;
+	}
+
 	TStrongObjectPtr<UContentBrowserFileDataSource> PoFileDataSource;
+
+	TArray<TTuple<FDelegateHandle, FCanEditFileDelegate>> CanEditFileOverrides;
 };
 
 IMPLEMENT_MODULE(FPortableObjectFileDataSourceModule, PortableObjectFileDataSource);
diff --git a/Engine/Plugins/Editor/Localization/PortableObjectFileDataSource/Source/PortableObjectFileDataSource/Public/IPortableObjectFileDataSourceModule.h b/Engine/Plugins/Editor/Localization/PortableObjectFileDataSource/Source/PortableObjectFileDataSource/Public/IPortableObjectFileDataSourceModule.h
new file mode 100644
--- /dev/null
+++ b/Engine/Plugins/Editor/Localization/PortableObjectFileDataSource/Source/PortableObjectFileDataSource/Public/IPortableObjectFileDataSourceModule.h
@@ -0,0 +1,33 @@
+// Copyright Epic Games, Inc. All Rights Reserved.
+
+#pragma once
+
+#include "Delegates/Delegate.h"
+#include "Modules/ModuleManager.h"
+
+class IPortableObjectFileDataSourceModule : public IModuleInterface
+{
+public:
+	static IPortableObjectFileDataSourceModule& Get()
+	{
+		return FModuleManager::LoadModuleChecked<IPortableObjectFileDataSourceModule>("PortableObjectFileDataSource");
+	}
+	
+	static IPortableObjectFileDataSourceModule* GetPtr()
+	{
+		return FModuleManager::GetModulePtr<IPortableObjectFileDataSourceModule>("PortableObjectFileDataSource");
+	}
+
+	using FCanEditFileDelegate = TDelegate<bool(const FName InFilePath, const FString& InFilename, FText* OutErrorMsg)>;
+
+	/**
+	 * Register an override handler for the "CanEdit" logic for Portable Object files.
+	 * Each handler is queried in turn to see if they all allow editing of the file, and editing is refused if any return false.
+	 */
+	virtual FDelegateHandle RegisterCanEditFileOverride(FCanEditFileDelegate&& Delegate) = 0;
+	
+	/**
+	 * Unregister an override handler for the "CanEdit" logic for Portable Object files, as previously registered by RegisterCanEditFileOverride.
+	 */
+	virtual void UnregisterCanEditFileOverride(const FDelegateHandle& Handle) = 0;
+};
