Fix plugin discovery during refresh and mount duplication issue in PortableObjectFileDataSourceModule. Replaced the if (!Plugin->bEnabled) check with a TSet of previously discovered plugins so disabled plugins are no longer treated as new. Also added a guard in OnContentPathMounted to skip re-mounting the same localization path, preventing the "File mount already registered" ensure.

diff --git a/Engine/Plugins/Editor/Localization/PortableObjectFileDataSource/Source/PortableObjectFileDataSource/Private/PortableObjectFileDataSourceModule.cpp b/Engine/Plugins/Editor/Localization/PortableObjectFileDataSource/Source/PortableObjectFileDataSource/Private/PortableObjectFileDataSourceModule.cpp
--- a/Engine/Plugins/Editor/Localization/PortableObjectFileDataSource/Source/PortableObjectFileDataSource/Private/PortableObjectFileDataSourceModule.cpp
+++ b/Engine/Plugins/Editor/Localization/PortableObjectFileDataSource/Source/PortableObjectFileDataSource/Private/PortableObjectFileDataSourceModule.cpp
@@ -132,9 +132,10 @@ private:
 		if (PoFileDataSource)
 		{
 			const FString LocalizationTargetFolder = InFilesystemPath / TEXT("Localization");
-			if (FPaths::DirectoryExists(LocalizationTargetFolder))
+			const FString VirtualMountPath = InAssetPath / TEXT("Localization");
+			if (FPaths::DirectoryExists(LocalizationTargetFolder) && !PoFileDataSource->HasFileMount(*VirtualMountPath))
 			{
-				PoFileDataSource->AddFileMount(*(InAssetPath / TEXT("Localization")), InFilesystemPath / TEXT("Localization"));
+				PoFileDataSource->AddFileMount(*VirtualMountPath, LocalizationTargetFolder);
 			}
 		}
 	}
diff --git a/Engine/Source/Runtime/Projects/Private/PluginManager.cpp b/Engine/Source/Runtime/Projects/Private/PluginManager.cpp
--- a/Engine/Source/Runtime/Projects/Private/PluginManager.cpp
+++ b/Engine/Source/Runtime/Projects/Private/PluginManager.cpp
@@ -572,10 +572,17 @@ void FPluginManager::RefreshPluginsList()
 	for (const FDiscoveredPluginMap::ElementType& PluginPair : AllPlugins)
 	{
 		const TSharedRef<FPlugin>& Plugin = DiscoveredPluginMapUtils::ResolvePluginFromMapVal(PluginPair.Value);
-		if (!Plugin->bEnabled) // if it's not enabled, it's a new plugin
+		
+		// Use Plugin->FileName to ensure uniqueness if two plugins share the same Name
+		const FString& PluginFilePath = Plugin->FileName;
+		uint32 FilePathHash = GetTypeHash(PluginFilePath);
+
+		if (!PreviouslyDiscoveredPlugins.ContainsByHash(FilePathHash, PluginFilePath))
 		{
-			const uint32 PluginNameHash = GetTypeHash(Plugin->GetName());
-			PluginsToConfigure.AddByHash(PluginNameHash, Plugin->GetName());
+			PluginsToConfigure.Add(Plugin->GetName());
+			
+			// store the file path to distinguish different plugin locations
+			PreviouslyDiscoveredPlugins.AddByHash(FilePathHash, PluginFilePath);
 		}
 
 #if WITH_EDITOR
diff --git a/Engine/Source/Runtime/Projects/Private/PluginManager.h b/Engine/Source/Runtime/Projects/Private/PluginManager.h
--- a/Engine/Source/Runtime/Projects/Private/PluginManager.h
+++ b/Engine/Source/Runtime/Projects/Private/PluginManager.h
@@ -374,6 +374,9 @@ private:
 
 	bool bPreloadedBinaries = false;
 
+	/** Track discovered plugin file paths to avoid treating disabled plugins as new */
+	TSet<FString> PreviouslyDiscoveredPlugins;
+
 	/** List of additional directory paths to search for plugins within */
 	TSet<FExternalPluginPath> ExternalPluginSources;
 
