[Backout] - CL39518357
#rnx
Original CL Desc
-----------------------------------------------------------------
Fix plugin discovery during refresh and mount duplication issue in PortableObjectFileDataSourceModule. Replaced the if (!Plugin->bEnabled) check with a TSet of previously discovered plugins so disabled plugins are no longer treated as new. Also added a guard in OnContentPathMounted to skip re-mounting the same localization path, preventing the "File mount already registered" ensure.

diff --git a/Engine/Plugins/Editor/Localization/PortableObjectFileDataSource/Source/PortableObjectFileDataSource/Private/PortableObjectFileDataSourceModule.cpp b/Engine/Plugins/Editor/Localization/PortableObjectFileDataSource/Source/PortableObjectFileDataSource/Private/PortableObjectFileDataSourceModule.cpp
--- a/Engine/Plugins/Editor/Localization/PortableObjectFileDataSource/Source/PortableObjectFileDataSource/Private/PortableObjectFileDataSourceModule.cpp
+++ b/Engine/Plugins/Editor/Localization/PortableObjectFileDataSource/Source/PortableObjectFileDataSource/Private/PortableObjectFileDataSourceModule.cpp
@@ -132,10 +132,9 @@ private:
 		if (PoFileDataSource)
 		{
 			const FString LocalizationTargetFolder = InFilesystemPath / TEXT("Localization");
-			const FString VirtualMountPath = InAssetPath / TEXT("Localization");
-			if (FPaths::DirectoryExists(LocalizationTargetFolder) && !PoFileDataSource->HasFileMount(*VirtualMountPath))
+			if (FPaths::DirectoryExists(LocalizationTargetFolder))
 			{
-				PoFileDataSource->AddFileMount(*VirtualMountPath, LocalizationTargetFolder);
+				PoFileDataSource->AddFileMount(*(InAssetPath / TEXT("Localization")), InFilesystemPath / TEXT("Localization"));
 			}
 		}
 	}
diff --git a/Engine/Source/Runtime/Projects/Private/PluginManager.cpp b/Engine/Source/Runtime/Projects/Private/PluginManager.cpp
--- a/Engine/Source/Runtime/Projects/Private/PluginManager.cpp
+++ b/Engine/Source/Runtime/Projects/Private/PluginManager.cpp
@@ -572,17 +572,10 @@ void FPluginManager::RefreshPluginsList()
 	for (const FDiscoveredPluginMap::ElementType& PluginPair : AllPlugins)
 	{
 		const TSharedRef<FPlugin>& Plugin = DiscoveredPluginMapUtils::ResolvePluginFromMapVal(PluginPair.Value);
-		
-		// Use Plugin->FileName to ensure uniqueness if two plugins share the same Name
-		const FString& PluginFilePath = Plugin->FileName;
-		uint32 FilePathHash = GetTypeHash(PluginFilePath);
-
-		if (!PreviouslyDiscoveredPlugins.ContainsByHash(FilePathHash, PluginFilePath))
+		if (!Plugin->bEnabled) // if it's not enabled, it's a new plugin
 		{
-			PluginsToConfigure.Add(Plugin->GetName());
-			
-			// store the file path to distinguish different plugin locations
-			PreviouslyDiscoveredPlugins.AddByHash(FilePathHash, PluginFilePath);
+			const uint32 PluginNameHash = GetTypeHash(Plugin->GetName());
+			PluginsToConfigure.AddByHash(PluginNameHash, Plugin->GetName());
 		}
 
 #if WITH_EDITOR
diff --git a/Engine/Source/Runtime/Projects/Private/PluginManager.h b/Engine/Source/Runtime/Projects/Private/PluginManager.h
--- a/Engine/Source/Runtime/Projects/Private/PluginManager.h
+++ b/Engine/Source/Runtime/Projects/Private/PluginManager.h
@@ -374,9 +374,6 @@ private:
 
 	bool bPreloadedBinaries = false;
 
-	/** Track discovered plugin file paths to avoid treating disabled plugins as new */
-	TSet<FString> PreviouslyDiscoveredPlugins;
-
 	/** List of additional directory paths to search for plugins within */
 	TSet<FExternalPluginPath> ExternalPluginSources;
 
