[PortableObjectFileDataSource] Prevent duplicate file mounts in OnContentPathMounted

diff --git a/Engine/Plugins/Editor/Localization/PortableObjectFileDataSource/Source/PortableObjectFileDataSource/Private/PortableObjectFileDataSourceModule.cpp b/Engine/Plugins/Editor/Localization/PortableObjectFileDataSource/Source/PortableObjectFileDataSource/Private/PortableObjectFileDataSourceModule.cpp
--- a/Engine/Plugins/Editor/Localization/PortableObjectFileDataSource/Source/PortableObjectFileDataSource/Private/PortableObjectFileDataSourceModule.cpp
+++ b/Engine/Plugins/Editor/Localization/PortableObjectFileDataSource/Source/PortableObjectFileDataSource/Private/PortableObjectFileDataSourceModule.cpp
@@ -132,9 +132,10 @@ private:
 		if (PoFileDataSource)
 		{
 			const FString LocalizationTargetFolder = InFilesystemPath / TEXT("Localization");
-			if (FPaths::DirectoryExists(LocalizationTargetFolder))
+			const FString VirtualMountPath = InAssetPath / TEXT("Localization");
+			if (FPaths::DirectoryExists(LocalizationTargetFolder) && !PoFileDataSource->HasFileMount(*VirtualMountPath))
 			{
-				PoFileDataSource->AddFileMount(*(InAssetPath / TEXT("Localization")), InFilesystemPath / TEXT("Localization"));
+				PoFileDataSource->AddFileMount(*VirtualMountPath, LocalizationTargetFolder);
 			}
 		}
 	}
